Try including file containing definition MacGraphicsSwitchingWidget relies on to fix non-unity CIS error.

diff --git a/Engine/Plugins/Editor/MacGraphicsSwitching/Source/MacGraphicsSwitching/Private/MacGraphicsSwitchingWidget.h b/Engine/Plugins/Editor/MacGraphicsSwitching/Source/MacGraphicsSwitching/Private/MacGraphicsSwitchingWidget.h
--- a/Engine/Plugins/Editor/MacGraphicsSwitching/Source/MacGraphicsSwitching/Private/MacGraphicsSwitchingWidget.h
+++ b/Engine/Plugins/Editor/MacGraphicsSwitching/Source/MacGraphicsSwitching/Private/MacGraphicsSwitchingWidget.h
@@ -3,6 +3,7 @@
 #pragma once
 
 #include "IMacGraphicsSwitchingModule.h"
+#include "PropertyHandle.h"
 #include "Slate.h"
 #include "SlateStyle.h"
 
