Fixed a crash on MacGraphicsSwitching initialization in commandlet mode.

diff --git a/Engine/Plugins/Editor/MacGraphicsSwitching/Source/MacGraphicsSwitching/Private/MacGraphicsSwitchingModule.cpp b/Engine/Plugins/Editor/MacGraphicsSwitching/Source/MacGraphicsSwitching/Private/MacGraphicsSwitchingModule.cpp
--- a/Engine/Plugins/Editor/MacGraphicsSwitching/Source/MacGraphicsSwitching/Private/MacGraphicsSwitchingModule.cpp
+++ b/Engine/Plugins/Editor/MacGraphicsSwitching/Source/MacGraphicsSwitching/Private/MacGraphicsSwitchingModule.cpp
@@ -104,7 +104,7 @@ void FMacGraphicsSwitchingModule::StartupModule()
 	
 	// Register settings
 	ISettingsModule* SettingsModule = ISettingsModule::Get();
-	if( ensure( SettingsModule != nullptr ) )
+	if( SettingsModule != nullptr )
 	{
 		SettingsModule->RegisterSettings( "Editor", "Plugins", "MacGraphicsSwitching",
 										 LOCTEXT( "MacGraphicsSwitchingSettingsName", "Graphics Switching"),
