Add include to fix non-unity Mac

diff --git a/Engine/Plugins/Editor/MacGraphicsSwitching/Source/MacGraphicsSwitching/Private/MacGraphicsSwitchingSettingsDetails.cpp b/Engine/Plugins/Editor/MacGraphicsSwitching/Source/MacGraphicsSwitching/Private/MacGraphicsSwitchingSettingsDetails.cpp
--- a/Engine/Plugins/Editor/MacGraphicsSwitching/Source/MacGraphicsSwitching/Private/MacGraphicsSwitchingSettingsDetails.cpp
+++ b/Engine/Plugins/Editor/MacGraphicsSwitching/Source/MacGraphicsSwitching/Private/MacGraphicsSwitchingSettingsDetails.cpp
@@ -1,6 +1,7 @@
 // Copyright 1998-2014 Epic Games, Inc. All Rights Reserved.
 
 #include "MacGraphicsSwitchingModule.h"
+#include "MacGraphicsSwitchingSettings.h"
 #include "MacGraphicsSwitchingSettingsDetails.h"
 #include "MacGraphicsSwitchingWidget.h"
 
