Duplicate 4102827 & 4102970:

Fix copyright dates & the drop-down for selecting the GPU should only appear in the Editor window when someone really wants it - I find it handy but it isn't worth showing by default.

diff --git a/Engine/Plugins/Editor/MacGraphicsSwitching/Source/MacGraphicsSwitching/MacGraphicsSwitching.Build.cs b/Engine/Plugins/Editor/MacGraphicsSwitching/Source/MacGraphicsSwitching/MacGraphicsSwitching.Build.cs
--- a/Engine/Plugins/Editor/MacGraphicsSwitching/Source/MacGraphicsSwitching/MacGraphicsSwitching.Build.cs
+++ b/Engine/Plugins/Editor/MacGraphicsSwitching/Source/MacGraphicsSwitching/MacGraphicsSwitching.Build.cs
@@ -1,4 +1,4 @@
-// Copyright 1998-2016 Epic Games, Inc. All Rights Reserved.
+// Copyright 1998-2018 Epic Games, Inc. All Rights Reserved.
 
 namespace UnrealBuildTool.Rules
 {
diff --git a/Engine/Plugins/Editor/MacGraphicsSwitching/Source/MacGraphicsSwitching/Private/MacGraphicsSwitchingModule.cpp b/Engine/Plugins/Editor/MacGraphicsSwitching/Source/MacGraphicsSwitching/Private/MacGraphicsSwitchingModule.cpp
--- a/Engine/Plugins/Editor/MacGraphicsSwitching/Source/MacGraphicsSwitching/Private/MacGraphicsSwitchingModule.cpp
+++ b/Engine/Plugins/Editor/MacGraphicsSwitching/Source/MacGraphicsSwitching/Private/MacGraphicsSwitchingModule.cpp
@@ -1,4 +1,4 @@
-// Copyright 1998-2016 Epic Games, Inc. All Rights Reserved.
+// Copyright 1998-2018 Epic Games, Inc. All Rights Reserved.
 
 #include "MacGraphicsSwitchingModule.h"
 #include "IMacGraphicsSwitchingModule.h"
@@ -65,15 +65,19 @@ void FMacGraphicsSwitchingModule::Initialize( TSharedPtr<SWindow> InRootWindow,
 
 		FMacGraphicsSwitchingStyle::Initialize();
 		
-		NotificationBarExtender = MakeShareable( new FExtender() );
-		NotificationBarExtender->AddToolBarExtension("Start",
-													 EExtensionHook::After,
-													 nullptr,
-													 FToolBarExtensionDelegate::CreateRaw(this, &FMacGraphicsSwitchingModule::AddGraphicsSwitcher));
-		
-		FLevelEditorModule& LevelEditorModule = FModuleManager::LoadModuleChecked<FLevelEditorModule>("LevelEditor");
-		LevelEditorModule.GetNotificationBarExtensibilityManager()->AddExtender( NotificationBarExtender );
-		LevelEditorModule.BroadcastNotificationBarChanged();
+		bool MacUseAutomaticGraphicsSwitching = false;
+		if ( GConfig->GetBool(TEXT("/Script/MacGraphicsSwitching.MacGraphicsSwitchingSettings"), TEXT("bShowGraphicsSwitching"), MacUseAutomaticGraphicsSwitching, GEditorSettingsIni) && MacUseAutomaticGraphicsSwitching )
+		{
+			NotificationBarExtender = MakeShareable( new FExtender() );
+			NotificationBarExtender->AddToolBarExtension("Start",
+														 EExtensionHook::After,
+														 nullptr,
+														 FToolBarExtensionDelegate::CreateRaw(this, &FMacGraphicsSwitchingModule::AddGraphicsSwitcher));
+			
+			FLevelEditorModule& LevelEditorModule = FModuleManager::LoadModuleChecked<FLevelEditorModule>("LevelEditor");
+			LevelEditorModule.GetNotificationBarExtensibilityManager()->AddExtender( NotificationBarExtender );
+			LevelEditorModule.BroadcastNotificationBarChanged();
+		}
 	}
 }
 
diff --git a/Engine/Plugins/Editor/MacGraphicsSwitching/Source/MacGraphicsSwitching/Private/MacGraphicsSwitchingModule.h b/Engine/Plugins/Editor/MacGraphicsSwitching/Source/MacGraphicsSwitching/Private/MacGraphicsSwitchingModule.h
--- a/Engine/Plugins/Editor/MacGraphicsSwitching/Source/MacGraphicsSwitching/Private/MacGraphicsSwitchingModule.h
+++ b/Engine/Plugins/Editor/MacGraphicsSwitching/Source/MacGraphicsSwitching/Private/MacGraphicsSwitchingModule.h
@@ -1,3 +1,3 @@
-// Copyright 1998-2016 Epic Games, Inc. All Rights Reserved.
+// Copyright 1998-2018 Epic Games, Inc. All Rights Reserved.
 
 #pragma once
diff --git a/Engine/Plugins/Editor/MacGraphicsSwitching/Source/MacGraphicsSwitching/Private/MacGraphicsSwitchingSettings.cpp b/Engine/Plugins/Editor/MacGraphicsSwitching/Source/MacGraphicsSwitching/Private/MacGraphicsSwitchingSettings.cpp
--- a/Engine/Plugins/Editor/MacGraphicsSwitching/Source/MacGraphicsSwitching/Private/MacGraphicsSwitchingSettings.cpp
+++ b/Engine/Plugins/Editor/MacGraphicsSwitching/Source/MacGraphicsSwitching/Private/MacGraphicsSwitchingSettings.cpp
@@ -1,4 +1,4 @@
-// Copyright 1998-2016 Epic Games, Inc. All Rights Reserved.
+// Copyright 1998-2018 Epic Games, Inc. All Rights Reserved.
 
 #include "MacGraphicsSwitchingSettings.h"
 #include "MacGraphicsSwitchingModule.h"
diff --git a/Engine/Plugins/Editor/MacGraphicsSwitching/Source/MacGraphicsSwitching/Private/MacGraphicsSwitchingSettings.h b/Engine/Plugins/Editor/MacGraphicsSwitching/Source/MacGraphicsSwitching/Private/MacGraphicsSwitchingSettings.h
--- a/Engine/Plugins/Editor/MacGraphicsSwitching/Source/MacGraphicsSwitching/Private/MacGraphicsSwitchingSettings.h
+++ b/Engine/Plugins/Editor/MacGraphicsSwitching/Source/MacGraphicsSwitching/Private/MacGraphicsSwitchingSettings.h
@@ -1,4 +1,4 @@
-// Copyright 1998-2016 Epic Games, Inc. All Rights Reserved.
+// Copyright 1998-2018 Epic Games, Inc. All Rights Reserved.
 #pragma once
 
 #include "CoreMinimal.h"
@@ -14,4 +14,7 @@ public:
 	
 	UPROPERTY(Config, EditAnywhere, Category=RHI, meta = (ConfigRestartRequired=true))
 	int32 RendererID;
+	
+	UPROPERTY(Config, EditAnywhere, Category=RHI, meta = (DisplayName = "Show Editor GPU Selector", ToolTip = "Adds a drop-down menu to the main Editor window allowing GPU selection. Requires restarting the Editor.", ConfigRestartRequired=true))
+	bool bShowGraphicsSwitching;
 };
diff --git a/Engine/Plugins/Editor/MacGraphicsSwitching/Source/MacGraphicsSwitching/Private/MacGraphicsSwitchingSettingsDetails.cpp b/Engine/Plugins/Editor/MacGraphicsSwitching/Source/MacGraphicsSwitching/Private/MacGraphicsSwitchingSettingsDetails.cpp
--- a/Engine/Plugins/Editor/MacGraphicsSwitching/Source/MacGraphicsSwitching/Private/MacGraphicsSwitchingSettingsDetails.cpp
+++ b/Engine/Plugins/Editor/MacGraphicsSwitching/Source/MacGraphicsSwitching/Private/MacGraphicsSwitchingSettingsDetails.cpp
@@ -1,4 +1,4 @@
-// Copyright 1998-2016 Epic Games, Inc. All Rights Reserved.
+// Copyright 1998-2018 Epic Games, Inc. All Rights Reserved.
 
 #include "MacGraphicsSwitchingSettingsDetails.h"
 #include "MacGraphicsSwitchingModule.h"
diff --git a/Engine/Plugins/Editor/MacGraphicsSwitching/Source/MacGraphicsSwitching/Private/MacGraphicsSwitchingSettingsDetails.h b/Engine/Plugins/Editor/MacGraphicsSwitching/Source/MacGraphicsSwitching/Private/MacGraphicsSwitchingSettingsDetails.h
--- a/Engine/Plugins/Editor/MacGraphicsSwitching/Source/MacGraphicsSwitching/Private/MacGraphicsSwitchingSettingsDetails.h
+++ b/Engine/Plugins/Editor/MacGraphicsSwitching/Source/MacGraphicsSwitching/Private/MacGraphicsSwitchingSettingsDetails.h
@@ -1,4 +1,4 @@
-// Copyright 1998-2016 Epic Games, Inc. All Rights Reserved.
+// Copyright 1998-2018 Epic Games, Inc. All Rights Reserved.
 
 #pragma once
 
@@ -13,4 +13,4 @@ public:
 
 	/** IDetailCustomization interface */
 	virtual void CustomizeDetails( IDetailLayoutBuilder& DetailLayout ) override;
-};
\ No newline at end of file
+};
diff --git a/Engine/Plugins/Editor/MacGraphicsSwitching/Source/MacGraphicsSwitching/Private/MacGraphicsSwitchingStyle.cpp b/Engine/Plugins/Editor/MacGraphicsSwitching/Source/MacGraphicsSwitching/Private/MacGraphicsSwitchingStyle.cpp
--- a/Engine/Plugins/Editor/MacGraphicsSwitching/Source/MacGraphicsSwitching/Private/MacGraphicsSwitchingStyle.cpp
+++ b/Engine/Plugins/Editor/MacGraphicsSwitching/Source/MacGraphicsSwitching/Private/MacGraphicsSwitchingStyle.cpp
@@ -1,4 +1,4 @@
-// Copyright 1998-2016 Epic Games, Inc. All Rights Reserved.
+// Copyright 1998-2018 Epic Games, Inc. All Rights Reserved.
 
 #include "MacGraphicsSwitchingStyle.h"
 #include "MacGraphicsSwitchingModule.h"
diff --git a/Engine/Plugins/Editor/MacGraphicsSwitching/Source/MacGraphicsSwitching/Private/MacGraphicsSwitchingStyle.h b/Engine/Plugins/Editor/MacGraphicsSwitching/Source/MacGraphicsSwitching/Private/MacGraphicsSwitchingStyle.h
--- a/Engine/Plugins/Editor/MacGraphicsSwitching/Source/MacGraphicsSwitching/Private/MacGraphicsSwitchingStyle.h
+++ b/Engine/Plugins/Editor/MacGraphicsSwitching/Source/MacGraphicsSwitching/Private/MacGraphicsSwitchingStyle.h
@@ -1,4 +1,4 @@
-// Copyright 1998-2016 Epic Games, Inc. All Rights Reserved.
+// Copyright 1998-2018 Epic Games, Inc. All Rights Reserved.
 
 #pragma once
 
@@ -15,4 +15,4 @@ public:
 
 private:
 	static TSharedPtr< class FSlateStyleSet > StyleSet;
-};
\ No newline at end of file
+};
diff --git a/Engine/Plugins/Editor/MacGraphicsSwitching/Source/MacGraphicsSwitching/Private/MacGraphicsSwitchingWidget.cpp b/Engine/Plugins/Editor/MacGraphicsSwitching/Source/MacGraphicsSwitching/Private/MacGraphicsSwitchingWidget.cpp
--- a/Engine/Plugins/Editor/MacGraphicsSwitching/Source/MacGraphicsSwitching/Private/MacGraphicsSwitchingWidget.cpp
+++ b/Engine/Plugins/Editor/MacGraphicsSwitching/Source/MacGraphicsSwitching/Private/MacGraphicsSwitchingWidget.cpp
@@ -1,4 +1,4 @@
-// Copyright 1998-2016 Epic Games, Inc. All Rights Reserved.
+// Copyright 1998-2018 Epic Games, Inc. All Rights Reserved.
 
 #include "MacGraphicsSwitchingWidget.h"
 #include "MacGraphicsSwitchingModule.h"
diff --git a/Engine/Plugins/Editor/MacGraphicsSwitching/Source/MacGraphicsSwitching/Private/MacGraphicsSwitchingWidget.h b/Engine/Plugins/Editor/MacGraphicsSwitching/Source/MacGraphicsSwitching/Private/MacGraphicsSwitchingWidget.h
--- a/Engine/Plugins/Editor/MacGraphicsSwitching/Source/MacGraphicsSwitching/Private/MacGraphicsSwitchingWidget.h
+++ b/Engine/Plugins/Editor/MacGraphicsSwitching/Source/MacGraphicsSwitching/Private/MacGraphicsSwitchingWidget.h
@@ -1,4 +1,4 @@
-// Copyright 1998-2016 Epic Games, Inc. All Rights Reserved.
+// Copyright 1998-2018 Epic Games, Inc. All Rights Reserved.
 
 #pragma once
 
