PR #4909: Update copyright notices to 2018 (Contributed by [REDACTED_PERSON])

diff --git a/Engine/Plugins/Editor/MacGraphicsSwitching/Source/MacGraphicsSwitching/Public/IMacGraphicsSwitchingModule.h b/Engine/Plugins/Editor/MacGraphicsSwitching/Source/MacGraphicsSwitching/Public/IMacGraphicsSwitchingModule.h
--- a/Engine/Plugins/Editor/MacGraphicsSwitching/Source/MacGraphicsSwitching/Public/IMacGraphicsSwitchingModule.h
+++ b/Engine/Plugins/Editor/MacGraphicsSwitching/Source/MacGraphicsSwitching/Public/IMacGraphicsSwitchingModule.h
@@ -1,4 +1,4 @@
-// Copyright 1998-2016 Epic Games, Inc. All Rights Reserved.
+// Copyright 1998-2018 Epic Games, Inc. All Rights Reserved.
 
 #pragma once
 
diff --git a/Engine/Plugins/Runtime/AppleARKit/Source/AppleARKit/Private/AppleARKitTrackable.cpp b/Engine/Plugins/Runtime/AppleARKit/Source/AppleARKit/Private/AppleARKitTrackable.cpp
--- a/Engine/Plugins/Runtime/AppleARKit/Source/AppleARKit/Private/AppleARKitTrackable.cpp
+++ b/Engine/Plugins/Runtime/AppleARKit/Source/AppleARKit/Private/AppleARKitTrackable.cpp
@@ -1,4 +1,4 @@
-// Copyright 1998-2017 Epic Games, Inc. All Rights Reserved.
+// Copyright 1998-2018 Epic Games, Inc. All Rights Reserved.
 
 #include "AppleARKitTrackable.h"
 #include "AppleARKitTextures.h"
diff --git a/Engine/Source/Runtime/AudioMixer/Private/DSP/BufferVectorOperations.cpp b/Engine/Source/Runtime/AudioMixer/Private/DSP/BufferVectorOperations.cpp
--- a/Engine/Source/Runtime/AudioMixer/Private/DSP/BufferVectorOperations.cpp
+++ b/Engine/Source/Runtime/AudioMixer/Private/DSP/BufferVectorOperations.cpp
@@ -1,4 +1,4 @@
-// Copyright 1998-2017 Epic Games, Inc. All Rights Reserved.
+// Copyright 1998-2018 Epic Games, Inc. All Rights Reserved.
 
 #include "DSP/BufferVectorOperations.h"
 
diff --git a/Engine/Source/Runtime/AudioMixer/Public/DSP/BufferVectorOperations.h b/Engine/Source/Runtime/AudioMixer/Public/DSP/BufferVectorOperations.h
--- a/Engine/Source/Runtime/AudioMixer/Public/DSP/BufferVectorOperations.h
+++ b/Engine/Source/Runtime/AudioMixer/Public/DSP/BufferVectorOperations.h
@@ -1,4 +1,4 @@
-// Copyright 1998-2017 Epic Games, Inc. All Rights Reserved.
+// Copyright 1998-2018 Epic Games, Inc. All Rights Reserved.
 
 #pragma once
 
diff --git a/Engine/Source/Runtime/AugmentedReality/Public/ARSkyLight.h b/Engine/Source/Runtime/AugmentedReality/Public/ARSkyLight.h
--- a/Engine/Source/Runtime/AugmentedReality/Public/ARSkyLight.h
+++ b/Engine/Source/Runtime/AugmentedReality/Public/ARSkyLight.h
@@ -1,4 +1,4 @@
-// Copyright 1998-2017 Epic Games, Inc. All Rights Reserved.
+// Copyright 1998-2018 Epic Games, Inc. All Rights Reserved.
 
 #pragma once
 
