Fix mac compile

diff --git a/Engine/Plugins/Editor/MacGraphicsSwitching/Source/MacGraphicsSwitching/Private/MacGraphicsSwitchingModule.cpp b/Engine/Plugins/Editor/MacGraphicsSwitching/Source/MacGraphicsSwitching/Private/MacGraphicsSwitchingModule.cpp
--- a/Engine/Plugins/Editor/MacGraphicsSwitching/Source/MacGraphicsSwitching/Private/MacGraphicsSwitchingModule.cpp
+++ b/Engine/Plugins/Editor/MacGraphicsSwitching/Source/MacGraphicsSwitching/Private/MacGraphicsSwitchingModule.cpp
@@ -76,7 +76,7 @@ void FMacGraphicsSwitchingModule::Initialize( TSharedPtr<SWindow> InRootWindow,
 			
 			FLevelEditorModule& LevelEditorModule = FModuleManager::LoadModuleChecked<FLevelEditorModule>("LevelEditor");
 			LevelEditorModule.GetNotificationBarExtensibilityManager()->AddExtender( NotificationBarExtender );
-			LevelEditorModule.BroadcastNotificationBarChanged();
+			LevelEditorModule.BroadcastTitleBarMessagesChanged();
 		}
 	}
 }
